Fix textual reference to PhysXVehicles plugin, should be ChaosVehicle plugin in the config.ini file

diff --git a/Templates/TP_VehicleAdv/Config/config.ini b/Templates/TP_VehicleAdv/Config/config.ini
--- a/Templates/TP_VehicleAdv/Config/config.ini
+++ b/Templates/TP_VehicleAdv/Config/config.ini
@@ -41,4 +41,4 @@ GameName=TP_VehicleAdv
 
 [FeaturePackSettings]
 CompileSource=false
-InstallMessage="VehicleAdv Feature Pack requires that the PhysXVehicles plugin be enabled and included in the PublicDependencyModuleNames range of your project's build.cs file. Before you can open the installed maps or blueprints you'll update your build.cs and recompile."
+InstallMessage="VehicleAdv Feature Pack requires that the ChaosVehiclesPlugin be enabled and included in the PublicDependencyModuleNames range of your project's build.cs file. Before you can open the installed maps or blueprints you'll update your build.cs and recompile."
