Fixed an incorrect default GameMode setting in the VehicleAdv template.
Updated VehicleAdv map to match the VehicleAdvBP template.

diff --git a/Templates/TP_VehicleAdv/Config/DefaultEngine.ini b/Templates/TP_VehicleAdv/Config/DefaultEngine.ini
--- a/Templates/TP_VehicleAdv/Config/DefaultEngine.ini
+++ b/Templates/TP_VehicleAdv/Config/DefaultEngine.ini
@@ -12,7 +12,7 @@ ThreePlayerSplitscreenLayout=FavorTop
 GameInstanceClass=/Script/Engine.GameInstance
 GameDefaultMap=/Game/VehicleTemplate/Maps/VehicleAdvExampleMap.VehicleAdvExampleMap
 ServerDefaultMap=/Engine/Maps/Entry
-GlobalDefaultGameMode=/Script/TP_VehicleAdv.TP_VehicleAdvGameMode
+GlobalDefaultGameMode=/Game/VehicleTemplate/Blueprints/BP_VehicleAdvGameMode.BP_VehicleAdvGameMode_C
 GlobalDefaultServerGameMode=None
 
 [/Script/Engine.RendererSettings]
