Project Templates - Removed extra console key binding to Caret from some of the templates.

#rnx

diff --git a/Templates/TP_Blank/Config/DefaultInput.ini b/Templates/TP_Blank/Config/DefaultInput.ini
--- a/Templates/TP_Blank/Config/DefaultInput.ini
+++ b/Templates/TP_Blank/Config/DefaultInput.ini
@@ -81,5 +81,4 @@ DefaultInputComponentClass=/Script/EnhancedInput.EnhancedInputComponent
 DefaultTouchInterface=/Engine/MobileResources/HUD/DefaultVirtualJoysticks.DefaultVirtualJoysticks
 -ConsoleKeys=Tilde
 +ConsoleKeys=Tilde
-+ConsoleKeys=Caret
 
diff --git a/Templates/TP_BlankBP/Config/DefaultInput.ini b/Templates/TP_BlankBP/Config/DefaultInput.ini
--- a/Templates/TP_BlankBP/Config/DefaultInput.ini
+++ b/Templates/TP_BlankBP/Config/DefaultInput.ini
@@ -81,5 +81,4 @@ DefaultInputComponentClass=/Script/EnhancedInput.EnhancedInputComponent
 DefaultTouchInterface=/Engine/MobileResources/HUD/DefaultVirtualJoysticks.DefaultVirtualJoysticks
 -ConsoleKeys=Tilde
 +ConsoleKeys=Tilde
-+ConsoleKeys=Caret
 
diff --git a/Templates/TP_ThirdPersonBP/Config/DefaultInput.ini b/Templates/TP_ThirdPersonBP/Config/DefaultInput.ini
--- a/Templates/TP_ThirdPersonBP/Config/DefaultInput.ini
+++ b/Templates/TP_ThirdPersonBP/Config/DefaultInput.ini
@@ -84,5 +84,4 @@ DefaultInputComponentClass=/Script/EnhancedInput.EnhancedInputComponent
 DefaultTouchInterface=/Engine/MobileResources/HUD/DefaultVirtualJoysticks.DefaultVirtualJoysticks
 -ConsoleKeys=Tilde
 +ConsoleKeys=Tilde
-+ConsoleKeys=Caret
 
diff --git a/Templates/TP_VehicleAdv/Config/DefaultInput.ini b/Templates/TP_VehicleAdv/Config/DefaultInput.ini
--- a/Templates/TP_VehicleAdv/Config/DefaultInput.ini
+++ b/Templates/TP_VehicleAdv/Config/DefaultInput.ini
@@ -144,6 +144,7 @@ DoubleClickTime=0.200000
 +AxisMappings=(AxisName="MoveRight",Scale=1.000000,Key=ValveIndex_Left_Thumbstick_X)
 +AxisMappings=(AxisName="MoveRight",Scale=1.000000,Key=ValveIndex_Right_Thumbstick_X)
 DefaultTouchInterface=/Game/Vehicles/Vehicle/LeftVirtualJoystickOnly.LeftVirtualJoystickOnly
+-ConsoleKeys=Tilde
 +ConsoleKeys=Tilde
 DefaultPlayerInputClass=/Script/EnhancedInput.EnhancedPlayerInput
 DefaultInputComponentClass=/Script/EnhancedInput.EnhancedInputComponent
diff --git a/Templates/TP_VehicleAdvBP/Config/DefaultInput.ini b/Templates/TP_VehicleAdvBP/Config/DefaultInput.ini
--- a/Templates/TP_VehicleAdvBP/Config/DefaultInput.ini
+++ b/Templates/TP_VehicleAdvBP/Config/DefaultInput.ini
@@ -113,5 +113,4 @@ DefaultInputComponentClass=/Script/EnhancedInput.EnhancedInputComponent
 DefaultTouchInterface=/Engine/MobileResources/HUD/DefaultVirtualJoysticks.DefaultVirtualJoysticks
 -ConsoleKeys=Tilde
 +ConsoleKeys=Tilde
-+ConsoleKeys=Caret
 
