Templates
- Added missing DefaultTouchInterface override to template projects that was causing duplicate touch interfaces on mobile.
#rnx

diff --git a/Templates/TP_FirstPerson/Config/DefaultInput.ini b/Templates/TP_FirstPerson/Config/DefaultInput.ini
--- a/Templates/TP_FirstPerson/Config/DefaultInput.ini
+++ b/Templates/TP_FirstPerson/Config/DefaultInput.ini
@@ -103,6 +103,7 @@ FOVScale=0.011110
 DoubleClickTime=0.200000
 DefaultPlayerInputClass=/Script/EnhancedInput.EnhancedPlayerInput
 DefaultInputComponentClass=/Script/EnhancedInput.EnhancedInputComponent
+DefaultTouchInterface=None
 -ConsoleKeys=Tilde
 +ConsoleKeys=Tilde
 
diff --git a/Templates/TP_VehicleAdv/Config/DefaultInput.ini b/Templates/TP_VehicleAdv/Config/DefaultInput.ini
--- a/Templates/TP_VehicleAdv/Config/DefaultInput.ini
+++ b/Templates/TP_VehicleAdv/Config/DefaultInput.ini
@@ -143,8 +143,8 @@ DoubleClickTime=0.200000
 +AxisMappings=(AxisName="MoveRight",Scale=1.000000,Key=OculusTouch_Right_Thumbstick_X)
 +AxisMappings=(AxisName="MoveRight",Scale=1.000000,Key=ValveIndex_Left_Thumbstick_X)
 +AxisMappings=(AxisName="MoveRight",Scale=1.000000,Key=ValveIndex_Right_Thumbstick_X)
-DefaultTouchInterface=/Game/Vehicles/Vehicle/LeftVirtualJoystickOnly.LeftVirtualJoystickOnly
 -ConsoleKeys=Tilde
 +ConsoleKeys=Tilde
 DefaultPlayerInputClass=/Script/EnhancedInput.EnhancedPlayerInput
 DefaultInputComponentClass=/Script/EnhancedInput.EnhancedInputComponent
+DefaultTouchInterface=None
\ No newline at end of file
