Update the Advanced Vehicle template to use Game Input instead of Raw Input for processing racing wheel inputs.

- Change the key mappings from raw input "Generic" USB keys to the Game Input Keys
- Remove the use of Raw Input from the template project
- Include GameInputRedist installer with packaging of the project.

#tests Cook and package template project. Use the logitech g920 racing wheel to ensure the inputs work as expected
Note: There is no PF that includes the vehicle template project

#virtualized

diff --git a/Templates/TP_VehicleAdv/Config/DefaultEngine.ini b/Templates/TP_VehicleAdv/Config/DefaultEngine.ini
--- a/Templates/TP_VehicleAdv/Config/DefaultEngine.ini
+++ b/Templates/TP_VehicleAdv/Config/DefaultEngine.ini
@@ -69,6 +69,22 @@ SoundCueCookQualityIndex=-1
 -TargetedRHIs=SF_VULKAN_SM5
 +TargetedRHIs=SF_VULKAN_SM6
 
-[/Script/MacTargetPlatform.MacTargetSettings]
--TargetedRHIs=SF_METAL_SM5
-+TargetedRHIs=SF_METAL_SM6
+; Include the GameInput redist files on windows targets
+
+[GameInput]
+IncludeRedistFiles=True
+
+[/Script/AndroidFileServerEditor.AndroidFileServerRuntimeSettings]
+bEnablePlugin=True
+bAllowNetworkConnection=True
+SecurityToken=DA0C08D74145274365E1FD8D5F4C9560
+bIncludeInShipping=False
+bAllowExternalStartInShipping=False
+bCompileAFSProject=False
+bUseCompression=False
+bLogFiles=False
+bReportStats=False
+ConnectionType=USBOnly
+bUseManualIPAddress=False
+ManualIPAddress=
+
diff --git a/Templates/TP_VehicleAdv/Config/DefaultInput.ini b/Templates/TP_VehicleAdv/Config/DefaultInput.ini
--- a/Templates/TP_VehicleAdv/Config/DefaultInput.ini
+++ b/Templates/TP_VehicleAdv/Config/DefaultInput.ini
@@ -1,6 +1,3 @@
-[/Script/RawInput.RawInputSettings]
-+DeviceConfigurations=(VendorID="0x046D",ProductID="0xC262",AxisProperties=((bEnabled=False),(Key=GenericUSBController_Axis1,Offset=-0.500000),(bEnabled=False),(Key=GenericUSBController_Axis2,bInverted=True,Offset=1.000000),(Key=GenericUSBController_Axis3,bInverted=True,Offset=1.000000)),ButtonProperties=((Key=GenericUSBController_Button1),(Key=GenericUSBController_Button2),(Key=GenericUSBController_Button3),(Key=GenericUSBController_Button4),(Key=GenericUSBController_Button5),(Key=GenericUSBController_Button6),(Key=GenericUSBController_Button7),(Key=GenericUSBController_Button8),(Key=GenericUSBController_Button9),(Key=GenericUSBController_Button10),(Key=GenericUSBController_Button11),(Key=GenericUSBController_Button12)))
-
 [/Script/Engine.InputSettings]
 +AxisConfig=(AxisKeyName="Gamepad_LeftX",AxisProperties=(DeadZone=0.250000,Sensitivity=1.000000,Exponent=1.000000,bInvert=False))
 +AxisConfig=(AxisKeyName="Gamepad_LeftY",AxisProperties=(DeadZone=0.250000,Sensitivity=1.000000,Exponent=1.000000,bInvert=False))
@@ -147,4 +144,7 @@ DoubleClickTime=0.200000
 +ConsoleKeys=Tilde
 DefaultPlayerInputClass=/Script/EnhancedInput.EnhancedPlayerInput
 DefaultInputComponentClass=/Script/EnhancedInput.EnhancedInputComponent
-DefaultTouchInterface=None
\ No newline at end of file
+DefaultTouchInterface=None
+
+[GameInputPlatformSettings_Windows GameInputPlatformSettings]
+bProcessRacingWheel=True
\ No newline at end of file
diff --git a/Templates/TP_VehicleAdv/Source/TP_VehicleAdv/TP_VehicleAdvPlayerController.cpp b/Templates/TP_VehicleAdv/Source/TP_VehicleAdv/TP_VehicleAdvPlayerController.cpp
--- a/Templates/TP_VehicleAdv/Source/TP_VehicleAdv/TP_VehicleAdvPlayerController.cpp
+++ b/Templates/TP_VehicleAdv/Source/TP_VehicleAdv/TP_VehicleAdvPlayerController.cpp
@@ -78,6 +78,11 @@ void ATP_VehicleAdvPlayerController::SetupInputComponent()
 					Subsystem->AddMappingContext(CurrentContext, 0);
 				}
 			}
+
+			if (bUseSteeringWheelControls)
+			{
+				Subsystem->AddMappingContext(SteeringWheelInputMappingContext, 0);
+			}
 		}
 	}
 }
diff --git a/Templates/TP_VehicleAdv/Source/TP_VehicleAdv/Variant_TimeTrial/TimeTrialPlayerController.cpp b/Templates/TP_VehicleAdv/Source/TP_VehicleAdv/Variant_TimeTrial/TimeTrialPlayerController.cpp
--- a/Templates/TP_VehicleAdv/Source/TP_VehicleAdv/Variant_TimeTrial/TimeTrialPlayerController.cpp
+++ b/Templates/TP_VehicleAdv/Source/TP_VehicleAdv/Variant_TimeTrial/TimeTrialPlayerController.cpp
@@ -98,6 +98,11 @@ void ATimeTrialPlayerController::SetupInputComponent()
 					Subsystem->AddMappingContext(CurrentContext, 0);
 				}
 			}
+
+			if (bUseSteeringWheelControls)
+			{
+				Subsystem->AddMappingContext(SteeringWheelInputMappingContext, 0);
+			}
 		}
 	}
 }
diff --git a/Templates/TP_VehicleAdv/TP_VehicleAdv.uproject b/Templates/TP_VehicleAdv/TP_VehicleAdv.uproject
--- a/Templates/TP_VehicleAdv/TP_VehicleAdv.uproject
+++ b/Templates/TP_VehicleAdv/TP_VehicleAdv.uproject
@@ -27,7 +27,7 @@
 			]
 		},
 		{
-			"Name": "RawInput",
+			"Name": "GameInputWindows",
 			"Enabled": true
 		}
 	]
diff --git a/Templates/TP_VehicleAdvBP/Config/DefaultEngine.ini b/Templates/TP_VehicleAdvBP/Config/DefaultEngine.ini
--- a/Templates/TP_VehicleAdvBP/Config/DefaultEngine.ini
+++ b/Templates/TP_VehicleAdvBP/Config/DefaultEngine.ini
@@ -74,4 +74,8 @@ SoundCueCookQualityIndex=-1
 
 [/Script/MacTargetPlatform.MacTargetSettings]
 -TargetedRHIs=SF_METAL_SM5
-+TargetedRHIs=SF_METAL_SM6
\ No newline at end of file
++TargetedRHIs=SF_METAL_SM6
+
+; Include the GameInput redist files on windows targets
+[GameInput]
+IncludeRedistFiles=True
\ No newline at end of file
diff --git a/Templates/TP_VehicleAdvBP/Config/DefaultInput.ini b/Templates/TP_VehicleAdvBP/Config/DefaultInput.ini
--- a/Templates/TP_VehicleAdvBP/Config/DefaultInput.ini
+++ b/Templates/TP_VehicleAdvBP/Config/DefaultInput.ini
@@ -1,6 +1,3 @@
-[/Script/RawInput.RawInputSettings]
-+DeviceConfigurations=(VendorID="0x046D",ProductID="0xC262",AxisProperties=((bEnabled=False),(Key=GenericUSBController_Axis1,Offset=-0.500000),(bEnabled=False),(Key=GenericUSBController_Axis2,bInverted=True,Offset=1.000000),(Key=GenericUSBController_Axis3,bInverted=True,Offset=1.000000)),ButtonProperties=((Key=GenericUSBController_Button1),(Key=GenericUSBController_Button2),(Key=GenericUSBController_Button3),(Key=GenericUSBController_Button4),(Key=GenericUSBController_Button5),(Key=GenericUSBController_Button6),(Key=GenericUSBController_Button7),(Key=GenericUSBController_Button8),(Key=GenericUSBController_Button9),(Key=GenericUSBController_Button10),(Key=GenericUSBController_Button11),(Key=GenericUSBController_Button12)))
-
 [/Script/Engine.InputSettings]
 -AxisConfig=(AxisKeyName="Gamepad_LeftX",AxisProperties=(DeadZone=0.25,Exponent=1.f,Sensitivity=1.f))
 -AxisConfig=(AxisKeyName="Gamepad_LeftY",AxisProperties=(DeadZone=0.25,Exponent=1.f,Sensitivity=1.f))
@@ -114,3 +111,6 @@ DefaultTouchInterface=None
 -ConsoleKeys=Tilde
 +ConsoleKeys=Tilde
 
+[GameInputPlatformSettings_Windows GameInputPlatformSettings]
+bProcessRacingWheel=True
+
diff --git a/Templates/TP_VehicleAdvBP/TP_VehicleAdvBP.uproject b/Templates/TP_VehicleAdvBP/TP_VehicleAdvBP.uproject
--- a/Templates/TP_VehicleAdvBP/TP_VehicleAdvBP.uproject
+++ b/Templates/TP_VehicleAdvBP/TP_VehicleAdvBP.uproject
@@ -5,7 +5,7 @@
 	"Description": "",
 	"Plugins": [
 		{
-			"Name": "RawInput",
+			"Name": "GameInputWindows",
 			"Enabled": true
 		},
 		{
